Fixed anim viewport pinned command widgets being added incorrectly (when not holding shift)

diff --git a/Engine/Source/Editor/PinnedCommandList/Private/SPinnedCommandList.cpp b/Engine/Source/Editor/PinnedCommandList/Private/SPinnedCommandList.cpp
--- a/Engine/Source/Editor/PinnedCommandList/Private/SPinnedCommandList.cpp
+++ b/Engine/Source/Editor/PinnedCommandList/Private/SPinnedCommandList.cpp
@@ -483,7 +483,7 @@ void SPinnedCommandList::LoadSettings()
 				}
 				else if(Command.Type == EPinnedCommandListType::CustomWidget)
 				{
-					AddCustomWidget(Command.Name);
+					AddCustomWidget_Internal(Command.Name);
 				}
 			}
 		}
@@ -591,7 +591,10 @@ void SPinnedCommandList::HandleCustomWidgetInteraction(FName InCustomWidgetIdent
 
 void SPinnedCommandList::AddCustomWidget(FName InCustomWidgetIdentifier)
 {
-	AddCustomWidget_Internal(InCustomWidgetIdentifier);
+	if(FSlateApplication::Get().GetModifierKeys().AreModifersDown(EModifierKey::Shift))
+	{
+		AddCustomWidget_Internal(InCustomWidgetIdentifier);
+	}
 }
 
 TSharedPtr<SCommand> SPinnedCommandList::AddCustomWidget_Internal(FName InCustomWidgetIdentifier, const TSharedPtr<const FUICommandList_Pinnable>& InUICommandListPinnable)
